Remove unnecessary log

diff --git a/Engine/Plugins/Runtime/Firebase/Source/Firebase.Build.cs b/Engine/Plugins/Runtime/Firebase/Source/Firebase.Build.cs
--- a/Engine/Plugins/Runtime/Firebase/Source/Firebase.Build.cs
+++ b/Engine/Plugins/Runtime/Firebase/Source/Firebase.Build.cs
@@ -16,10 +16,7 @@ namespace UnrealBuildTool.Rules
             });
 
             PublicDefinitions.Add("WITH_FIREBASE=1");
-
 			string PluginPath = Utils.MakePathRelativeTo(ModuleDirectory, Target.RelativeEnginePath);
-	
-	System.Console.WriteLine("Plugin path: {0}, ModuleDirectory: {1}", PluginPath, ModuleDirectory);
 
             if (Target.Platform == UnrealTargetPlatform.Android)
 			{
