Update Firebase to deal with IMMUTABLE crash

#android

diff --git a/Engine/Build/Android/Java/aar-imports-ant.txt b/Engine/Build/Android/Java/aar-imports-ant.txt
--- a/Engine/Build/Android/Java/aar-imports-ant.txt
+++ b/Engine/Build/Android/Java/aar-imports-ant.txt
@@ -1,6 +1,6 @@
 repositories $(ANDROID_HOME)/extras
 repositories $(ENGINEDIR)/Source/ThirdParty/Android/extras
 com.google.android.gms,play-services-auth,11.0.4
-com.google.android.gms,play-services-games,11.0.4
+com.google.android.gms,play-services-games,23.0.0
 com.google.android.gms,play-services-nearby,11.0.4
 com.google.android.gms,play-services-plus,11.0.4
diff --git a/Engine/Build/Android/Java/aar-imports.txt b/Engine/Build/Android/Java/aar-imports.txt
--- a/Engine/Build/Android/Java/aar-imports.txt
+++ b/Engine/Build/Android/Java/aar-imports.txt
@@ -1,6 +1,6 @@
 repositories $(ANDROID_HOME)/extras
 repositories $(ENGINEDIR)/Source/ThirdParty/Android/extras
 com.google.android.gms,play-services-auth,11.8.0
-com.google.android.gms,play-services-games,11.8.0
+com.google.android.gms,play-services-games,23.0.0
 com.google.android.gms,play-services-nearby,11.8.0
 com.google.android.gms,play-services-plus,11.8.0
diff --git a/Engine/Plugins/Online/Android/OnlineSubsystemGooglePlay/Source/OnlineSubsystemGooglePlay_UPL.xml b/Engine/Plugins/Online/Android/OnlineSubsystemGooglePlay/Source/OnlineSubsystemGooglePlay_UPL.xml
--- a/Engine/Plugins/Online/Android/OnlineSubsystemGooglePlay/Source/OnlineSubsystemGooglePlay_UPL.xml
+++ b/Engine/Plugins/Online/Android/OnlineSubsystemGooglePlay/Source/OnlineSubsystemGooglePlay_UPL.xml
@@ -101,7 +101,6 @@ import java.security.NoSuchAlgorithmException;
 				GoogleApiClient.Builder gbuilder = new GoogleApiClient.Builder(this);
 				gbuilder.addConnectionCallbacks(this);
 				gbuilder.addOnConnectionFailedListener(this);
-				gbuilder.addApiIfAvailable(Games.API, Games.SCOPE_GAMES);
 				googleClient = gbuilder.build();
 			}
 			catch (Exception e)
diff --git a/Engine/Plugins/Runtime/Firebase/Source/Firebase.upl.xml b/Engine/Plugins/Runtime/Firebase/Source/Firebase.upl.xml
--- a/Engine/Plugins/Runtime/Firebase/Source/Firebase.upl.xml
+++ b/Engine/Plugins/Runtime/Firebase/Source/Firebase.upl.xml
@@ -20,15 +20,6 @@
 	<androidManifestUpdates>
 		<if condition="bEnabled">
 			<true>
-				<addElements tag="application">
-					<service
-					  android:name="com.epicgames.unreal.notifications.EpicFirebaseInstanceIDService"
-					  android:exported="false" >
-						<intent-filter>
-							<action android:name="com.google.firebase.INSTANCE_ID_EVENT"/>
-						</intent-filter>
-					</service>
-				</addElements>
 				<setElement result="ServiceElement" value="service" />
 				<addAttribute tag="$ServiceElement" name="android:name" value="$S(FirebaseService)" />
 				<addAttribute tag="$ServiceElement" name="android:exported" value="false" />
@@ -58,8 +49,8 @@
 			<true>
 				<!-- Needs to be same version number as play-services -->
 				<insert>dependencies {
-	implementation 'com.google.firebase:firebase-core:11.8.0'
-	implementation 'com.google.firebase:firebase-messaging:11.8.0'
+	implementation 'com.google.firebase:firebase-core:21.1.1'
+	implementation 'com.google.firebase:firebase-messaging:23.0.8'
 }
 
 // init firebase
diff --git a/Engine/Plugins/Runtime/Firebase/Source/Java/notifications/EpicFirebaseInstanceIDService.java b/Engine/Plugins/Runtime/Firebase/Source/Java/notifications/EpicFirebaseInstanceIDService.java
deleted file mode 100644
diff --git a/Engine/Plugins/Runtime/Firebase/Source/Java/notifications/EpicFirebaseMessagingService.java b/Engine/Plugins/Runtime/Firebase/Source/Java/notifications/EpicFirebaseMessagingService.java
--- a/Engine/Plugins/Runtime/Firebase/Source/Java/notifications/EpicFirebaseMessagingService.java
+++ b/Engine/Plugins/Runtime/Firebase/Source/Java/notifications/EpicFirebaseMessagingService.java
@@ -7,6 +7,7 @@ import android.content.Context;
 import android.content.Intent;
 import android.net.Uri;
 import android.os.Build;
+import android.content.SharedPreferences;
 import androidx.annotation.NonNull;
 import androidx.annotation.Nullable;
 import androidx.core.app.NotificationCompat;
@@ -248,4 +249,96 @@ public class EpicFirebaseMessagingService extends FirebaseMessagingService {
 		JSONObject o = new JSONObject(payload);
 		return o.optString(ATTR_TYPE);
 	}
+
+	// Formally in EpicFirebaseInstanceIDService.java
+	private static final String PREFS_FILE_FIREBASE = "com.epicgames.firebase";
+	private static final String KEY_FIREBASE_TOKEN = "firebasetoken";
+	private static final String KEY_IS_UPDATED_TOKEN = "isUpdatedToken";
+	private static final String KEY_IS_REGISTERED = "isRegistered";
+
+	@Override
+	public void onNewToken(@NonNull String firebaseToken) {
+		Log.debug("Refreshed Firebase token: " + firebaseToken);
+		if (TextUtils.isEmpty(firebaseToken)) {
+			Log.error("Firebase token is empty or null");
+		} else {
+			saveFirebaseToken(this, firebaseToken);
+		}
+	}
+
+	private static void saveFirebaseToken(@NonNull Context context, @NonNull String firebaseToken) {
+		Log.debug("Firebase token to save : " + firebaseToken);
+		String storedToken = getFirebaseTokenFromCache(context);
+		boolean isUpdatedToken = !TextUtils.isEmpty(storedToken);
+		SharedPreferences sharedPreferences = context.getSharedPreferences(PREFS_FILE_FIREBASE, Context.MODE_PRIVATE);
+		SharedPreferences.Editor editor = sharedPreferences.edit();
+		Log.debug("Firebase token isUpdated : " + isUpdatedToken);
+		editor.putBoolean(KEY_IS_UPDATED_TOKEN, isUpdatedToken).apply();
+		editor.putBoolean(KEY_IS_REGISTERED, false).apply();
+		editor.putString(KEY_FIREBASE_TOKEN, firebaseToken).apply();
+	}
+
+	@SuppressWarnings("unused")
+	static boolean isFirebaseTokenUpdated(@NonNull Context context) {
+		SharedPreferences sharedPreferences = context.getSharedPreferences(PREFS_FILE_FIREBASE, Context.MODE_PRIVATE);
+		boolean isUpdated = sharedPreferences.getBoolean(KEY_IS_UPDATED_TOKEN, true);
+		Log.debug("Firebase token isUpdatedToken is " + isUpdated);
+		return isUpdated;
+	}
+
+	public static boolean isFirebaseTokenRegistered(@NonNull Context context) {
+		SharedPreferences sharedPreferences = context.getSharedPreferences(PREFS_FILE_FIREBASE, Context.MODE_PRIVATE);
+		return sharedPreferences.getBoolean(KEY_IS_REGISTERED, false);
+	}
+
+	public static void setFirebaseTokenRegistered(@NonNull Context context, boolean isRegistered) {
+		Log.debug("Firebase token isRegistered setting to " + isRegistered);
+		SharedPreferences sharedPreferences = context.getSharedPreferences(PREFS_FILE_FIREBASE, Context.MODE_PRIVATE);
+		SharedPreferences.Editor editor = sharedPreferences.edit();
+		editor.putBoolean(KEY_IS_REGISTERED, isRegistered);
+		editor.apply();
+	}
+
+	public static void unregisterFirebaseToken(@NonNull Context context) {
+		setFirebaseTokenRegistered(context, false);
+		SharedPreferences sharedPreferences = context.getSharedPreferences(PREFS_FILE_FIREBASE, Context.MODE_PRIVATE);
+		SharedPreferences.Editor editor = sharedPreferences.edit();
+		editor.remove(KEY_FIREBASE_TOKEN);
+		editor.apply();
+		Log.debug("Firebase token cleared");
+	}
+
+	private static String getFirebaseTokenFromCache(@NonNull Context context) {
+		SharedPreferences sharedPreferences = context.getSharedPreferences(PREFS_FILE_FIREBASE, Context.MODE_PRIVATE);
+		return sharedPreferences.getString(KEY_FIREBASE_TOKEN, null);
+	}
+	
+	@Nullable
+	public static String getFirebaseToken(@NonNull Context context) {
+		String token = getFirebaseTokenFromCache(context);
+		Log.debug("Firebase token retrieved from cache: " + token);
+		if(TextUtils.isEmpty(token)) {
+			/*
+			Task<String> tokenTask = FirebaseMessaging.getInstance().getToken();
+			tokenTask.addOnCompleteListener(new OnCompleteListener<String>() {
+				@Override
+				public void onComplete(@NonNull Task<String> task) {
+					if (task.isSuccessful()) {
+						String token = task.getResult();
+						if(!TextUtils.isEmpty(token)) {
+							Log.debug("Firebase token retrieved from Firebase: " + token);
+							saveFirebaseToken(context, token);
+						}
+					}
+				}
+			});
+			// wait on task
+			Tasks.await(tokenTask);
+			token = getFirebaseTokenFromCache(context);
+			token = (token == null) ? "" : token;
+			*/
+		}
+		return token;
+	}
+
 }
