[Android] Bringing back code for getting a new Firebase token

diff --git a/Engine/Plugins/Runtime/Firebase/Source/Java/notifications/EpicFirebaseMessagingService.java b/Engine/Plugins/Runtime/Firebase/Source/Java/notifications/EpicFirebaseMessagingService.java
--- a/Engine/Plugins/Runtime/Firebase/Source/Java/notifications/EpicFirebaseMessagingService.java
+++ b/Engine/Plugins/Runtime/Firebase/Source/Java/notifications/EpicFirebaseMessagingService.java
@@ -17,6 +17,10 @@ import com.epicgames.unreal.GameActivity;
 import com.epicgames.unreal.GameApplication;
 import com.epicgames.unreal.LocalNotificationReceiver;
 import com.epicgames.unreal.Logger;
+import com.google.android.gms.tasks.OnCompleteListener;
+import com.google.android.gms.tasks.Task;
+import com.google.android.gms.tasks.Tasks;
+import com.google.firebase.messaging.FirebaseMessaging;
 import com.google.firebase.messaging.FirebaseMessagingService;
 import com.google.firebase.messaging.RemoteMessage;
 
@@ -27,6 +31,7 @@ import java.util.Map;
 import java.util.Objects;
 import java.util.Random;
 
+import static com.epicgames.unreal.GameActivity.Get;
 import static com.epicgames.unreal.GameActivity.LocalNotificationGetID;
 import static com.epicgames.unreal.LocalNotificationReceiver.KEY_LOCAL_NOTIFICATION_ACTION;
 import static com.epicgames.unreal.LocalNotificationReceiver.KEY_LOCAL_NOTIFICATION_BODY;
@@ -267,7 +272,7 @@ public class EpicFirebaseMessagingService extends FirebaseMessagingService {
 	}
 
 	private static void saveFirebaseToken(@NonNull Context context, @NonNull String firebaseToken) {
-		Log.debug("Firebase token to save : " + firebaseToken);
+		Log.debug("Saving Firebase token");
 		String storedToken = getFirebaseTokenFromCache(context);
 		boolean isUpdatedToken = !TextUtils.isEmpty(storedToken);
 		SharedPreferences sharedPreferences = context.getSharedPreferences(PREFS_FILE_FIREBASE, Context.MODE_PRIVATE);
@@ -316,9 +321,7 @@ public class EpicFirebaseMessagingService extends FirebaseMessagingService {
 	@Nullable
 	public static String getFirebaseToken(@NonNull Context context) {
 		String token = getFirebaseTokenFromCache(context);
-		Log.debug("Firebase token retrieved from cache: " + token);
 		if(TextUtils.isEmpty(token)) {
-			/*
 			Task<String> tokenTask = FirebaseMessaging.getInstance().getToken();
 			tokenTask.addOnCompleteListener(new OnCompleteListener<String>() {
 				@Override
@@ -326,18 +329,28 @@ public class EpicFirebaseMessagingService extends FirebaseMessagingService {
 					if (task.isSuccessful()) {
 						String token = task.getResult();
 						if(!TextUtils.isEmpty(token)) {
-							Log.debug("Firebase token retrieved from Firebase: " + token);
+							Log.debug("Firebase token retrieved from Firebase");
 							saveFirebaseToken(context, token);
 						}
 					}
 				}
 			});
 			// wait on task
-			Tasks.await(tokenTask);
-			token = getFirebaseTokenFromCache(context);
-			token = (token == null) ? "" : token;
-			*/
+			try {
+				Tasks.await(tokenTask);
+				token = getFirebaseTokenFromCache(context);
+				token = (token == null) ? "" : token;
+			} catch (Exception e) {
+				Log.error("Failed to retrieve Firebase token", e);
+			}
+		} else {
+			Log.debug("Firebase token retrieved from cache");
 		}
+
+		if (!Get().nativeIsShippingBuild()) {
+			Log.debug("Firebase token is " + token);
+		}
+
 		return token;
 	}
 
