Firebase notifications integration for iOS

diff --git a/Engine/Plugins/Runtime/Firebase/Source/IOS/notifications/EpicFirebaseIOSNotifications.cpp b/Engine/Plugins/Runtime/Firebase/Source/IOS/notifications/EpicFirebaseIOSNotifications.cpp
new file mode 100644
--- /dev/null
+++ b/Engine/Plugins/Runtime/Firebase/Source/IOS/notifications/EpicFirebaseIOSNotifications.cpp
@@ -0,0 +1,122 @@
+// Copyright Epic Games, Inc. All Rights Reserved.
+
+#include "EpicFirebaseIOSNotifications.h"
+
+#if PLATFORM_IOS && WITH_IOS_FIREBASE_INTEGRATION
+
+THIRD_PARTY_INCLUDES_START
+_Pragma("clang diagnostic push")
+_Pragma("clang diagnostic ignored \"-Wobjc-property-no-attribute\"")
+#include "ThirdParty/IOS/include/Firebase.h"
+_Pragma("clang diagnostic pop")
+THIRD_PARTY_INCLUDES_END
+// Module header, needed for log category
+#include "Firebase.h"
+#include "IOS/IOSAppDelegate.h"
+
+bool FFirebaseIOSNotifications::bIsInitialized = false;
+FString FFirebaseIOSNotifications::IOSFirebaseToken;
+
+@interface IOSAppDelegate (FirebaseHandling) <FIRMessagingDelegate>
+
+-(void)SetupFirebase;
+-(void)UpdateFirebaseToken : (UInt64) Timeout;
+
+@end
+
+@implementation IOSAppDelegate (FirebaseHandling)
+
+-(void)SetupFirebase
+{
+    [FIRApp configure];
+    [FIRMessaging messaging].delegate = self;
+    [UNUserNotificationCenter currentNotificationCenter].delegate = self;
+    
+    UNAuthorizationOptions authOptions = UNAuthorizationOptionAlert | UNAuthorizationOptionSound | UNAuthorizationOptionBadge;
+    [[UNUserNotificationCenter currentNotificationCenter] requestAuthorizationWithOptions:authOptions
+        completionHandler:^(BOOL granted, NSError * _Nullable error) 
+     {
+        if (granted)
+        {
+            UE_LOG(LogFirebase, Log, TEXT("Firebase authorization granted"));
+        }
+        else
+        {
+            UE_LOG(LogFirebase, Log, TEXT("Firebase authorization denied"));
+        }
+    }];
+
+    [[UIApplication sharedApplication] registerForRemoteNotifications];
+}
+
+-(void)messaging:(FIRMessaging *)messaging didReceiveRegistrationToken:(NSString *)fcmToken 
+{
+    NSDictionary *dataDict = [NSDictionary dictionaryWithObject:fcmToken forKey:@"token"];
+    [[NSNotificationCenter defaultCenter] postNotificationName:
+     @"FCMToken" object:nil userInfo:dataDict];
+    
+    FString Token = FString(fcmToken);
+    FFirebaseIOSNotifications::SetFirebaseToken(Token);
+#if !UE_BUILD_SHIPPING
+    UE_LOG(LogFirebase, Log, TEXT("Firebase Token : %s"), *Token);
+#endif
+}
+
+- (void)UpdateFirebaseToken : (UInt64) Timeout
+{
+    dispatch_semaphore_t updateTokenSemaphore = dispatch_semaphore_create(0);
+    // wrapped in dispatch_async to avoid locking up if we're on the main thread
+    dispatch_async(dispatch_get_global_queue(DISPATCH_QUEUE_PRIORITY_DEFAULT, 0), ^(void)
+    {
+        [[FIRInstanceID instanceID] instanceIDWithHandler:^(FIRInstanceIDResult * _Nullable result, NSError * _Nullable error)
+         {
+            if (error == nil && result != nil)
+            {
+                FString Token = FString(result.token);
+                FFirebaseIOSNotifications::SetFirebaseToken(Token);
+#if !UE_BUILD_SHIPPING
+                UE_LOG(LogFirebase, Log, TEXT("Firebase Token : %s"), *Token);
+#endif
+                dispatch_semaphore_signal(updateTokenSemaphore);
+            }
+        }];
+    });
+    dispatch_semaphore_wait(updateTokenSemaphore, dispatch_time(DISPATCH_TIME_NOW, Timeout));
+}
+@end
+
+void FFirebaseIOSNotifications::Initialize(uint64 TokenQueryTimeoutNanoseconds)
+{
+    if (!bIsInitialized)
+    {
+        [[IOSAppDelegate GetDelegate] SetupFirebase];
+        [[IOSAppDelegate GetDelegate] UpdateFirebaseToken:TokenQueryTimeoutNanoseconds];
+        bIsInitialized = true;
+    }
+}
+
+void FFirebaseIOSNotifications::SetFirebaseToken(FString Token)
+{
+    @synchronized ([IOSAppDelegate GetDelegate])
+    {
+        IOSFirebaseToken = Token;
+    }
+}
+
+FString FFirebaseIOSNotifications::GetFirebaseToken()
+{
+    FString Token;
+    @synchronized ([IOSAppDelegate GetDelegate])
+    {
+        Token = IOSFirebaseToken;
+    }
+    
+    if (Token.IsEmpty())
+    {
+        UE_LOG(LogFirebase, Log, TEXT("Firebase Token is empty"));
+    }
+    
+    return Token;
+}
+
+#endif // PLATFORM_IOS
diff --git a/Engine/Plugins/Runtime/Firebase/Source/IOS/notifications/EpicFirebaseIOSNotifications.h b/Engine/Plugins/Runtime/Firebase/Source/IOS/notifications/EpicFirebaseIOSNotifications.h
new file mode 100644
--- /dev/null
+++ b/Engine/Plugins/Runtime/Firebase/Source/IOS/notifications/EpicFirebaseIOSNotifications.h
@@ -0,0 +1,23 @@
+// Copyright Epic Games, Inc. All Rights Reserved.
+
+#pragma once
+
+#ifndef WITH_IOS_FIREBASE_INTEGRATION
+#define WITH_IOS_FIREBASE_INTEGRATION 0
+#endif
+
+#if PLATFORM_IOS && WITH_IOS_FIREBASE_INTEGRATION
+
+struct FIREBASE_API FFirebaseIOSNotifications
+{
+private:
+    static bool bIsInitialized;
+    static FString IOSFirebaseToken;
+    
+public:
+    static void Initialize(uint64 TokenQueryTimeoutNanoseconds);
+    static void SetFirebaseToken(FString Token);
+    static FString GetFirebaseToken();
+};
+
+#endif // PLATFORM_IOS
