Functionality for disabling Firebase auto init and enabling it at runtime. Storing Firebase token locally

diff --git a/Engine/Plugins/Runtime/Firebase/Source/Firebase.upl.xml b/Engine/Plugins/Runtime/Firebase/Source/Firebase.upl.xml
--- a/Engine/Plugins/Runtime/Firebase/Source/Firebase.upl.xml
+++ b/Engine/Plugins/Runtime/Firebase/Source/Firebase.upl.xml
@@ -6,6 +6,7 @@
 		<log text="Firebase init"/>
 
 		<setBoolFromProperty result="bEnabled" ini="Engine" section="Firebase" property="FirebaseEnabled" default="false"/>
+		<setBoolFromProperty result="bDisableAutoInit" ini="Engine" section="Firebase" property="FirebaseDisableAutoInit" default="false"/>
 		<setStringFromProperty result="FirebaseService" ini="Engine" section="Firebase" property="FirebaseService" default="com.epicgames.unreal.notifications.EpicFirebaseMessagingService"/>
 		<if condition="bEnabled">
 			<true>
@@ -28,6 +29,14 @@
 							<action android:name="com.google.firebase.MESSAGING_EVENT"/>
 						</intent-filter>
 				</addElements>
+				<if condition="bDisableAutoInit">
+					<true>
+						<addElements tag="$ServiceElement">
+							<meta-data android:name="firebase_messaging_auto_init_enabled" android:value="false" />
+							<meta-data android:name="firebase_analytics_collection_enabled" android:value="false" />
+						</addElements>
+					</true>
+				</if>
 				<addElement tag="application" name="ServiceElement" />
 			</true>
 		</if>
diff --git a/Engine/Plugins/Runtime/Firebase/Source/IOS/notifications/EpicFirebaseIOSNotifications.cpp b/Engine/Plugins/Runtime/Firebase/Source/IOS/notifications/EpicFirebaseIOSNotifications.cpp
--- a/Engine/Plugins/Runtime/Firebase/Source/IOS/notifications/EpicFirebaseIOSNotifications.cpp
+++ b/Engine/Plugins/Runtime/Firebase/Source/IOS/notifications/EpicFirebaseIOSNotifications.cpp
@@ -15,20 +15,26 @@ THIRD_PARTY_INCLUDES_END
 #include "IOS/IOSAppDelegate.h"
 
 bool FFirebaseIOSNotifications::bIsInitialized = false;
+bool FFirebaseIOSNotifications::bIsConfigured = false;
 FString FFirebaseIOSNotifications::IOSFirebaseToken;
+NSString* KEY_FIREBASE_TOKEN = @"firebasetoken";
 
 @interface IOSAppDelegate (FirebaseHandling) <FIRMessagingDelegate>
 
--(void)SetupFirebase;
+-(void)SetupFirebase : (Boolean) enableAnalytics;
 -(void)UpdateFirebaseToken : (UInt64) Timeout;
 
 @end
 
 @implementation IOSAppDelegate (FirebaseHandling)
 
--(void)SetupFirebase
+-(void)SetupFirebase  : (Boolean) enableAnalytics
 {
-    [FIRApp configure];
+    [FIRMessaging messaging].autoInitEnabled = YES;
+    if (enableAnalytics)
+    {
+        [FIRAnalytics setAnalyticsCollectionEnabled:YES];
+    }
     [FIRMessaging messaging].delegate = self;
     [UNUserNotificationCenter currentNotificationCenter].delegate = self;
     
@@ -46,7 +52,12 @@ FString FFirebaseIOSNotifications::IOSFirebaseToken;
         }
     }];
 
-    [[UIApplication sharedApplication] registerForRemoteNotifications];
+   [[UIApplication sharedApplication] registerForRemoteNotifications];
+}
+
+-(void)ConfigureFirebase
+{
+    [FIRApp configure];
 }
 
 -(void)messaging:(FIRMessaging *)messaging didReceiveRegistrationToken:(NSString *)fcmToken 
@@ -60,10 +71,23 @@ FString FFirebaseIOSNotifications::IOSFirebaseToken;
 #if !UE_BUILD_SHIPPING
     UE_LOG(LogFirebase, Log, TEXT("Firebase Token : %s"), *Token);
 #endif
+    [[NSUserDefaults standardUserDefaults] setObject:fcmToken forKey:KEY_FIREBASE_TOKEN];
+    [[NSUserDefaults standardUserDefaults] synchronize];
 }
 
 - (void)UpdateFirebaseToken : (UInt64) Timeout
 {
+    NSUserDefaults* UserDefaults = [NSUserDefaults standardUserDefaults];
+    if ([UserDefaults objectForKey:KEY_FIREBASE_TOKEN] != nil)
+    {
+        FString Token = FString([UserDefaults stringForKey:KEY_FIREBASE_TOKEN]);
+        FFirebaseIOSNotifications::SetFirebaseToken(Token);
+#if !UE_BUILD_SHIPPING
+        UE_LOG(LogFirebase, Log, TEXT("Retrieved Firebase Token from cache : %s"), *Token);
+#endif
+        return;
+    }
+    
     dispatch_semaphore_t updateTokenSemaphore = dispatch_semaphore_create(0);
     // wrapped in dispatch_async to avoid locking up if we're on the main thread
     dispatch_async(dispatch_get_global_queue(DISPATCH_QUEUE_PRIORITY_DEFAULT, 0), ^(void)
@@ -77,6 +101,9 @@ FString FFirebaseIOSNotifications::IOSFirebaseToken;
 #if !UE_BUILD_SHIPPING
                 UE_LOG(LogFirebase, Log, TEXT("Firebase Token : %s"), *Token);
 #endif
+                [[NSUserDefaults standardUserDefaults] setObject:firebaseToken forKey:KEY_FIREBASE_TOKEN];
+                [[NSUserDefaults standardUserDefaults] synchronize];
+                
                 dispatch_semaphore_signal(updateTokenSemaphore);
             }
         }];
@@ -85,11 +112,25 @@ FString FFirebaseIOSNotifications::IOSFirebaseToken;
 }
 @end
 
-void FFirebaseIOSNotifications::Initialize(uint64 TokenQueryTimeoutNanoseconds)
+void FFirebaseIOSNotifications::ConfigureFirebase()
+{
+    if (!bIsConfigured)
+    {
+        bIsConfigured = true;
+        [[IOSAppDelegate GetDelegate] ConfigureFirebase];
+    }
+}
+
+void FFirebaseIOSNotifications::Initialize(uint64 TokenQueryTimeoutNanoseconds, bool bEnableAnalytics)
 {
+    if (!bIsConfigured)
+    {
+        ConfigureFirebase();
+    }
+    
     if (!bIsInitialized)
     {
-        [[IOSAppDelegate GetDelegate] SetupFirebase];
+        [[IOSAppDelegate GetDelegate] SetupFirebase:bEnableAnalytics];
         [[IOSAppDelegate GetDelegate] UpdateFirebaseToken:TokenQueryTimeoutNanoseconds];
         bIsInitialized = true;
     }
diff --git a/Engine/Plugins/Runtime/Firebase/Source/IOS/notifications/EpicFirebaseIOSNotifications.h b/Engine/Plugins/Runtime/Firebase/Source/IOS/notifications/EpicFirebaseIOSNotifications.h
--- a/Engine/Plugins/Runtime/Firebase/Source/IOS/notifications/EpicFirebaseIOSNotifications.h
+++ b/Engine/Plugins/Runtime/Firebase/Source/IOS/notifications/EpicFirebaseIOSNotifications.h
@@ -12,10 +12,12 @@ struct FIREBASE_API FFirebaseIOSNotifications
 {
 private:
     static bool bIsInitialized;
+    static bool bIsConfigured;
     static FString IOSFirebaseToken;
     
 public:
-    static void Initialize(uint64 TokenQueryTimeoutNanoseconds);
+    static void ConfigureFirebase();
+    static void Initialize(uint64 TokenQueryTimeoutNanoseconds, bool bEnableAnalytics);
     static void SetFirebaseToken(FString Token);
     static FString GetFirebaseToken();
 };
diff --git a/Engine/Plugins/Runtime/Firebase/Source/Java/notifications/EpicFirebaseMessagingService.java b/Engine/Plugins/Runtime/Firebase/Source/Java/notifications/EpicFirebaseMessagingService.java
--- a/Engine/Plugins/Runtime/Firebase/Source/Java/notifications/EpicFirebaseMessagingService.java
+++ b/Engine/Plugins/Runtime/Firebase/Source/Java/notifications/EpicFirebaseMessagingService.java
@@ -23,6 +23,7 @@ import com.google.android.gms.tasks.Tasks;
 import com.google.firebase.messaging.FirebaseMessaging;
 import com.google.firebase.messaging.FirebaseMessagingService;
 import com.google.firebase.messaging.RemoteMessage;
+import com.google.firebase.analytics.FirebaseAnalytics;
 
 import org.json.JSONException;
 import org.json.JSONObject;
@@ -354,4 +355,10 @@ public class EpicFirebaseMessagingService extends FirebaseMessagingService {
 		return token;
 	}
 
+	public static void enableFirebaseAutoInit(@NonNull Context context, boolean enableAnalytics) {
+		FirebaseMessaging.getInstance().setAutoInitEnabled(true);
+		if (enableAnalytics) {
+			FirebaseAnalytics.getInstance(context).setAnalyticsCollectionEnabled(true);
+		}
+	}
 }
