IWYU fixup

diff --git a/Engine/Plugins/Runtime/Firebase/Source/IOS/notifications/EpicFirebaseIOSNotifications.h b/Engine/Plugins/Runtime/Firebase/Source/IOS/notifications/EpicFirebaseIOSNotifications.h
--- a/Engine/Plugins/Runtime/Firebase/Source/IOS/notifications/EpicFirebaseIOSNotifications.h
+++ b/Engine/Plugins/Runtime/Firebase/Source/IOS/notifications/EpicFirebaseIOSNotifications.h
@@ -2,6 +2,8 @@
 
 #pragma once
 
+#include "CoreMinimal.h"
+
 #ifndef WITH_IOS_FIREBASE_INTEGRATION
 #define WITH_IOS_FIREBASE_INTEGRATION 0
 #endif
