Fix for push notifications on new AGP version
Removed FirebaseInitProvider if auto init is disabled

diff --git a/Engine/Plugins/Runtime/Firebase/Source/Firebase.upl.xml b/Engine/Plugins/Runtime/Firebase/Source/Firebase.upl.xml
--- a/Engine/Plugins/Runtime/Firebase/Source/Firebase.upl.xml
+++ b/Engine/Plugins/Runtime/Firebase/Source/Firebase.upl.xml
@@ -35,6 +35,12 @@
 							<meta-data android:name="firebase_messaging_auto_init_enabled" android:value="false" />
 							<meta-data android:name="firebase_analytics_collection_enabled" android:value="false" />
 						</addElements>
+						<setElement result="removeFirebaseInit" value="provider" />
+						<addAttribute tag="$removeFirebaseInit" name="android:name" value="com.google.firebase.provider.FirebaseInitProvider" />
+						<addAttribute tag="$removeFirebaseInit" name="android:authorities" value="$S(AndroidPackageName).firebaseinitprovider" />
+						<addAttribute tag="$removeFirebaseInit" name="android:exported" value="false" />
+						<addAttribute tag="$removeFirebaseInit" name="tools:node" value="remove" />
+						<addElement tag="application" name="removeFirebaseInit" />
 					</true>
 				</if>
 				<addElement tag="application" name="ServiceElement" />
@@ -46,7 +52,7 @@
 		<if condition="bEnabled">
 			<true>
 				<insert>dependencies {
-	classpath 'com.google.gms:google-services:4.3.15' // google-services plugin
+	classpath 'com.google.gms:google-services:4.4.2' // google-services plugin
 }
 </insert>
 			</true>
diff --git a/Engine/Plugins/Runtime/Firebase/Source/Java/notifications/EpicFirebaseMessagingService.java b/Engine/Plugins/Runtime/Firebase/Source/Java/notifications/EpicFirebaseMessagingService.java
--- a/Engine/Plugins/Runtime/Firebase/Source/Java/notifications/EpicFirebaseMessagingService.java
+++ b/Engine/Plugins/Runtime/Firebase/Source/Java/notifications/EpicFirebaseMessagingService.java
@@ -20,6 +20,7 @@ import com.epicgames.unreal.Logger;
 import com.google.android.gms.tasks.OnCompleteListener;
 import com.google.android.gms.tasks.Task;
 import com.google.android.gms.tasks.Tasks;
+import com.google.firebase.FirebaseApp;
 import com.google.firebase.messaging.FirebaseMessaging;
 import com.google.firebase.messaging.FirebaseMessagingService;
 import com.google.firebase.messaging.RemoteMessage;
@@ -321,6 +322,7 @@ public class EpicFirebaseMessagingService extends FirebaseMessagingService {
 	
 	@Nullable
 	public static String getFirebaseToken(@NonNull Context context) {
+		EnsureFirebaseIsInitialized(context);
 		String token = getFirebaseTokenFromCache(context);
 		if(TextUtils.isEmpty(token)) {
 			Task<String> tokenTask = FirebaseMessaging.getInstance().getToken();
@@ -356,9 +358,17 @@ public class EpicFirebaseMessagingService extends FirebaseMessagingService {
 	}
 
 	public static void enableFirebaseAutoInit(@NonNull Context context, boolean enableAnalytics) {
+		EnsureFirebaseIsInitialized(context);
 		FirebaseMessaging.getInstance().setAutoInitEnabled(true);
 		if (enableAnalytics) {
 			FirebaseAnalytics.getInstance(context).setAnalyticsCollectionEnabled(true);
 		}
 	}
+
+	public static void EnsureFirebaseIsInitialized(@NonNull Context context)
+	{
+		if (FirebaseApp.getApps(context).isEmpty()) {
+			FirebaseApp.initializeApp(context);
+		}
+	}
 }
