Added path for initializing Firebase with custom config file path
Exposed functionality for deleting Firebase token

diff --git a/Engine/Plugins/Runtime/Firebase/Source/IOS/notifications/EpicFirebaseIOSNotifications.cpp b/Engine/Plugins/Runtime/Firebase/Source/IOS/notifications/EpicFirebaseIOSNotifications.cpp
--- a/Engine/Plugins/Runtime/Firebase/Source/IOS/notifications/EpicFirebaseIOSNotifications.cpp
+++ b/Engine/Plugins/Runtime/Firebase/Source/IOS/notifications/EpicFirebaseIOSNotifications.cpp
@@ -15,9 +15,9 @@ THIRD_PARTY_INCLUDES_END
 #include "IOS/IOSAppDelegate.h"
 
 bool FFirebaseIOSNotifications::bIsInitialized = false;
-bool FFirebaseIOSNotifications::bIsConfigured = false;
 FString FFirebaseIOSNotifications::IOSFirebaseToken;
 NSString* KEY_FIREBASE_TOKEN = @"firebasetoken";
+NSString* KEY_FIREBASE_PROJECT_ID = @"firebaseprojectid";
 
 @interface IOSAppDelegate (FirebaseHandling) <FIRMessagingDelegate>
 
@@ -28,13 +28,12 @@ NSString* KEY_FIREBASE_TOKEN = @"firebasetoken";
 
 @implementation IOSAppDelegate (FirebaseHandling)
 
--(void)SetupFirebase  : (Boolean) enableAnalytics
+-(void)SetupFirebase : (Boolean) enableAnalytics
 {
-    [FIRMessaging messaging].autoInitEnabled = YES;
-    if (enableAnalytics)
-    {
-        [FIRAnalytics setAnalyticsCollectionEnabled:YES];
-    }
+	if (enableAnalytics)
+	{
+		[FIRAnalytics setAnalyticsCollectionEnabled:YES];
+	}
     [FIRMessaging messaging].delegate = self;
     [UNUserNotificationCenter currentNotificationCenter].delegate = self;
     
@@ -60,8 +59,37 @@ NSString* KEY_FIREBASE_TOKEN = @"firebasetoken";
     [FIRApp configure];
 }
 
+-(void)ConfigureFirebaseWithCustomFile : (NSString*)fileName fileExtension:(NSString*)ext
+{
+	NSString* filePath = [[NSBundle mainBundle] pathForResource:fileName ofType:ext];
+	if (filePath == nil)
+	{
+		UE_LOG(LogFirebase, Warning, TEXT("Failed to find custom Firebase file, using default configuration"));
+		[self ConfigureFirebase];
+	}
+	else
+	{
+		FIROptions* options = [[FIROptions alloc] initWithContentsOfFile:filePath];
+		if (options == nil)
+		{
+			UE_LOG(LogFirebase, Warning, TEXT("Failed to parse custom Firebase options, using default configuration"));
+			[self ConfigureFirebase];
+		}
+		else
+		{
+			[FIRApp configureWithOptions:options];
+		}
+	}
+}
+
 -(void)messaging:(FIRMessaging *)messaging didReceiveRegistrationToken:(NSString *)fcmToken 
 {
+	if (fcmToken == nil)
+	{
+		// happens when deleting the token
+		return;
+	}
+	
     NSDictionary *dataDict = [NSDictionary dictionaryWithObject:fcmToken forKey:@"token"];
     [[NSNotificationCenter defaultCenter] postNotificationName:
      @"FCMToken" object:nil userInfo:dataDict];
@@ -69,15 +97,28 @@ NSString* KEY_FIREBASE_TOKEN = @"firebasetoken";
     FString Token = FString(fcmToken);
     FFirebaseIOSNotifications::SetFirebaseToken(Token);
 #if !UE_BUILD_SHIPPING
-    UE_LOG(LogFirebase, Log, TEXT("Firebase Token : %s"), *Token);
+    UE_LOG(LogFirebase, Log, TEXT("Firebase Token Refreshed : %s"), *Token);
 #endif
+	NSString* projectID = [[[FIRApp defaultApp] options] GCMSenderID];
     [[NSUserDefaults standardUserDefaults] setObject:fcmToken forKey:KEY_FIREBASE_TOKEN];
+	[[NSUserDefaults standardUserDefaults] setObject:projectID forKey:KEY_FIREBASE_PROJECT_ID];
     [[NSUserDefaults standardUserDefaults] synchronize];
 }
 
 - (void)UpdateFirebaseToken : (UInt64) Timeout
 {
     NSUserDefaults* UserDefaults = [NSUserDefaults standardUserDefaults];
+	NSString* currentProjectID = [[[FIRApp defaultApp] options] GCMSenderID];
+	if ([UserDefaults objectForKey:KEY_FIREBASE_PROJECT_ID] == nil ||
+		![[UserDefaults stringForKey:KEY_FIREBASE_PROJECT_ID] isEqualToString:currentProjectID])
+	{
+		if ([UserDefaults objectForKey:KEY_FIREBASE_TOKEN] != nil)
+		{
+			UE_LOG(LogFirebase, Log, TEXT("Firebase project changed, removing token"));
+			[UserDefaults removeObjectForKey:KEY_FIREBASE_TOKEN];
+		}
+	}
+	
     if ([UserDefaults objectForKey:KEY_FIREBASE_TOKEN] != nil)
     {
         FString Token = FString([UserDefaults stringForKey:KEY_FIREBASE_TOKEN]);
@@ -99,9 +140,11 @@ NSString* KEY_FIREBASE_TOKEN = @"firebasetoken";
                 FString Token = FString(firebaseToken);
                 FFirebaseIOSNotifications::SetFirebaseToken(Token);
 #if !UE_BUILD_SHIPPING
-                UE_LOG(LogFirebase, Log, TEXT("Firebase Token : %s"), *Token);
+                UE_LOG(LogFirebase, Log, TEXT("Firebase Token Queried : %s"), *Token);
 #endif
+				NSString* projectID = [[[FIRApp defaultApp] options] GCMSenderID];
                 [[NSUserDefaults standardUserDefaults] setObject:firebaseToken forKey:KEY_FIREBASE_TOKEN];
+                [[NSUserDefaults standardUserDefaults] setObject:projectID forKey:KEY_FIREBASE_PROJECT_ID];
                 [[NSUserDefaults standardUserDefaults] synchronize];
                 
                 dispatch_semaphore_signal(updateTokenSemaphore);
@@ -115,16 +158,30 @@ NSString* KEY_FIREBASE_TOKEN = @"firebasetoken";
 
 void FFirebaseIOSNotifications::ConfigureFirebase()
 {
-    if (!bIsConfigured)
+    if (!IsConfigured())
     {
-        bIsConfigured = true;
         [[IOSAppDelegate GetDelegate] ConfigureFirebase];
     }
 }
 
+void FFirebaseIOSNotifications::ConfigureFirebaseWithCustomFile(const FString& FileName)
+{
+	if (IsConfigured())
+	{
+		UE_LOG(LogFirebase, Warning, TEXT("Failed to configure Firebase with custom file. Firebase is already configured!"));
+		return;
+	}
+	
+	FString Name = FPaths::GetBaseFilename(FileName);
+	FString Extension = FPaths::GetExtension(FileName);
+	CFStringRef NameStr = FPlatformString::TCHARToCFString(*Name);
+	CFStringRef ExtensionStr = FPlatformString::TCHARToCFString(*Extension);
+	[[IOSAppDelegate GetDelegate] ConfigureFirebaseWithCustomFile:(NSString*)NameStr fileExtension:(NSString*)ExtensionStr];
+}
+
 void FFirebaseIOSNotifications::Initialize(uint64 TokenQueryTimeoutNanoseconds, bool bEnableAnalytics)
 {
-    if (!bIsConfigured)
+    if (!IsConfigured())
     {
         ConfigureFirebase();
     }
@@ -161,4 +218,33 @@ FString FFirebaseIOSNotifications::GetFirebaseToken()
     return Token;
 }
 
+void FFirebaseIOSNotifications::EnableFirebaseAutoInit()
+{
+	[FIRMessaging messaging].autoInitEnabled = YES;
+}
+
+bool FFirebaseIOSNotifications::IsConfigured()
+{
+	return [[FIRApp allApps] count] != 0;
+}
+
+void FFirebaseIOSNotifications::DeleteFirebaseToken()
+{
+	if (!IsConfigured())
+	{
+		UE_LOG(LogFirebase, Warning, TEXT("Trying to delete Firebase token, but Firebase is not configured."))
+		return;
+	}
+	
+	[[NSUserDefaults standardUserDefaults] removeObjectForKey:KEY_FIREBASE_TOKEN];
+	[[NSUserDefaults standardUserDefaults] synchronize];
+	
+	[[FIRMessaging messaging] deleteTokenWithCompletion:^(NSError *error)
+	 {
+#if !UE_BUILD_SHIPPING
+		NSLog(@"Error when deleting Firebase token: %@", error);
+#endif
+	}];
+}
+
 #endif // PLATFORM_IOS
diff --git a/Engine/Plugins/Runtime/Firebase/Source/IOS/notifications/EpicFirebaseIOSNotifications.h b/Engine/Plugins/Runtime/Firebase/Source/IOS/notifications/EpicFirebaseIOSNotifications.h
--- a/Engine/Plugins/Runtime/Firebase/Source/IOS/notifications/EpicFirebaseIOSNotifications.h
+++ b/Engine/Plugins/Runtime/Firebase/Source/IOS/notifications/EpicFirebaseIOSNotifications.h
@@ -16,12 +16,17 @@ private:
     static bool bIsInitialized;
     static bool bIsConfigured;
     static FString IOSFirebaseToken;
+	
+	static bool IsConfigured();
     
 public:
     static void ConfigureFirebase();
+	static void ConfigureFirebaseWithCustomFile(const FString& FileName);
+	static void EnableFirebaseAutoInit();
     static void Initialize(uint64 TokenQueryTimeoutNanoseconds, bool bEnableAnalytics);
     static void SetFirebaseToken(FString Token);
     static FString GetFirebaseToken();
+	static void DeleteFirebaseToken();
 };
 
 #endif // PLATFORM_IOS
diff --git a/Engine/Plugins/Runtime/Firebase/Source/Java/notifications/EpicFirebaseMessagingService.java b/Engine/Plugins/Runtime/Firebase/Source/Java/notifications/EpicFirebaseMessagingService.java
--- a/Engine/Plugins/Runtime/Firebase/Source/Java/notifications/EpicFirebaseMessagingService.java
+++ b/Engine/Plugins/Runtime/Firebase/Source/Java/notifications/EpicFirebaseMessagingService.java
@@ -21,6 +21,7 @@ import com.google.android.gms.tasks.OnCompleteListener;
 import com.google.android.gms.tasks.Task;
 import com.google.android.gms.tasks.Tasks;
 import com.google.firebase.FirebaseApp;
+import com.google.firebase.FirebaseOptions;
 import com.google.firebase.messaging.FirebaseMessaging;
 import com.google.firebase.messaging.FirebaseMessagingService;
 import com.google.firebase.messaging.RemoteMessage;
@@ -32,6 +33,10 @@ import org.json.JSONObject;
 import java.util.Map;
 import java.util.Objects;
 import java.util.Random;
+import java.io.InputStreamReader;
+import java.io.BufferedReader;
+import java.io.InputStream;
+import java.util.stream.Collectors;
 
 import static com.epicgames.unreal.GameActivity.Get;
 import static com.epicgames.unreal.GameActivity.LocalNotificationGetID;
@@ -260,6 +265,7 @@ public class EpicFirebaseMessagingService extends FirebaseMessagingService {
 	// Formally in EpicFirebaseInstanceIDService.java
 	private static final String PREFS_FILE_FIREBASE = "com.epicgames.firebase";
 	private static final String KEY_FIREBASE_TOKEN = "firebasetoken";
+	private static final String KEY_FIREBASE_PROJECT_ID = "firebaseProjectID";
 	private static final String KEY_IS_UPDATED_TOKEN = "isUpdatedToken";
 	private static final String KEY_IS_REGISTERED = "isRegistered";
 
@@ -270,6 +276,7 @@ public class EpicFirebaseMessagingService extends FirebaseMessagingService {
 			Log.error("Firebase token is empty or null");
 		} else {
 			saveFirebaseToken(this, firebaseToken);
+			setLastFirebaseProjectID(this, FirebaseApp.getInstance().getOptions().getGcmSenderId());
 		}
 	}
 
@@ -319,11 +326,39 @@ public class EpicFirebaseMessagingService extends FirebaseMessagingService {
 		SharedPreferences sharedPreferences = context.getSharedPreferences(PREFS_FILE_FIREBASE, Context.MODE_PRIVATE);
 		return sharedPreferences.getString(KEY_FIREBASE_TOKEN, null);
 	}
-	
+
+	private static String getLastFirebaseProjectID(@NonNull Context context)
+	{
+		SharedPreferences sharedPreferences = context.getSharedPreferences(PREFS_FILE_FIREBASE, Context.MODE_PRIVATE);
+		return sharedPreferences.getString(KEY_FIREBASE_PROJECT_ID, null);
+	}
+
+	private static void setLastFirebaseProjectID(@NonNull Context context, String projectID)
+	{
+		SharedPreferences sharedPreferences = context.getSharedPreferences(PREFS_FILE_FIREBASE, Context.MODE_PRIVATE);
+		SharedPreferences.Editor editor = sharedPreferences.edit();
+		editor.putString(KEY_FIREBASE_PROJECT_ID, projectID);
+		editor.apply();
+	}
+
 	@Nullable
 	public static String getFirebaseToken(@NonNull Context context) {
 		EnsureFirebaseIsInitialized(context);
-		String token = getFirebaseTokenFromCache(context);
+
+		String projectID = FirebaseApp.getInstance().getOptions().getGcmSenderId();
+		String lastProjectID = getLastFirebaseProjectID(context);
+		String token;
+		if (lastProjectID == null || !projectID.contentEquals(lastProjectID))
+		{
+			Log.debug("Firebase project changed, querying new token");
+			token = null;
+			saveFirebaseToken(context, "");
+		}
+		else
+		{
+			token = getFirebaseTokenFromCache(context);
+		}
+
 		if(TextUtils.isEmpty(token)) {
 			Task<String> tokenTask = FirebaseMessaging.getInstance().getToken();
 			tokenTask.addOnCompleteListener(new OnCompleteListener<String>() {
@@ -334,6 +369,7 @@ public class EpicFirebaseMessagingService extends FirebaseMessagingService {
 						if(!TextUtils.isEmpty(token)) {
 							Log.debug("Firebase token retrieved from Firebase");
 							saveFirebaseToken(context, token);
+							setLastFirebaseProjectID(context, projectID);
 						}
 					}
 				}
@@ -365,10 +401,79 @@ public class EpicFirebaseMessagingService extends FirebaseMessagingService {
 		}
 	}
 
-	public static void EnsureFirebaseIsInitialized(@NonNull Context context)
-	{
+	public static void EnsureFirebaseIsInitialized(@NonNull Context context) {
 		if (FirebaseApp.getApps(context).isEmpty()) {
 			FirebaseApp.initializeApp(context);
 		}
 	}
+
+	public static void deleteFirebaseToken(@NonNull Context context) {
+		saveFirebaseToken(context, "");
+		FirebaseMessaging.getInstance().deleteToken();
+	}
+
+	public static void initializeFirebaseWithCustomFile(@NonNull Context context, String fileName) {
+		if (!FirebaseApp.getApps(context).isEmpty())
+		{
+			Log.error("Failed to initialize Firebase with custom file. Firebase is already initialized!");
+			return;
+		}
+
+		try
+		{
+			InputStream inputStream = context.getAssets().open(fileName);
+			String result = new BufferedReader(new InputStreamReader(inputStream)).lines().collect(Collectors.joining("\n"));
+
+			/*
+				Expected JSON format :
+				{
+				  "project_info": {
+					"project_number": "X",
+					"firebase_url": "X",
+					"project_id": "X",
+					"storage_bucket": "X"
+				  },
+				  "client": [
+					{
+					  "client_info": {
+						"mobilesdk_app_id": "X"
+					  },
+					  "api_key": [
+						{
+						  "current_key": "X"
+						}
+					  ]
+					}
+				  ]
+				}
+			*/
+
+			JSONObject root = new JSONObject(result);
+			JSONObject projectInfo = root.getJSONObject("project_info");
+			JSONObject client = root.getJSONArray("client").getJSONObject(0);
+
+			String projectNumber = projectInfo.getString("project_number");
+			String firebaseUrl = projectInfo.getString("firebase_url");
+			String projectId = projectInfo.getString("project_id");
+			String storageBucket = projectInfo.getString("storage_bucket");
+			String apiKey = client.getJSONArray("api_key").getJSONObject(0).getString("current_key");
+			String appID = client.getJSONObject("client_info").getString("mobilesdk_app_id");
+
+			FirebaseOptions.Builder optionsBuilder = new FirebaseOptions.Builder();
+			optionsBuilder.setApiKey(apiKey);
+			optionsBuilder.setApplicationId(appID);
+			optionsBuilder.setDatabaseUrl(firebaseUrl);
+			optionsBuilder.setGcmSenderId(projectNumber);
+			optionsBuilder.setProjectId(projectId);
+			optionsBuilder.setStorageBucket(storageBucket);
+			FirebaseOptions options = optionsBuilder.build();
+
+			FirebaseApp.initializeApp(context, options);
+		}
+		catch(Exception e)
+		{
+			Log.error("Failed to initialize Firebase with custom file, using default initialization, exception : ", e);
+			FirebaseApp.initializeApp(context);
+		}
+	}
 }
