Firebase token update delegate. Query Firebase token even if there is one in cache

diff --git a/Engine/Plugins/Runtime/Firebase/Source/IOS/notifications/EpicFirebaseIOSNotifications.cpp b/Engine/Plugins/Runtime/Firebase/Source/IOS/notifications/EpicFirebaseIOSNotifications.cpp
--- a/Engine/Plugins/Runtime/Firebase/Source/IOS/notifications/EpicFirebaseIOSNotifications.cpp
+++ b/Engine/Plugins/Runtime/Firebase/Source/IOS/notifications/EpicFirebaseIOSNotifications.cpp
@@ -126,9 +126,10 @@ NSString* KEY_FIREBASE_PROJECT_ID = @"firebaseprojectid";
 #if !UE_BUILD_SHIPPING
         UE_LOG(LogFirebase, Log, TEXT("Retrieved Firebase Token from cache : %s"), *Token);
 #endif
-        return;
     }
     
+	// Query token from Firebase even if there is one in cache already
+	
     dispatch_semaphore_t updateTokenSemaphore = dispatch_semaphore_create(0);
     // wrapped in dispatch_async to avoid locking up if we're on the main thread
     dispatch_async(dispatch_get_global_queue(DISPATCH_QUEUE_PRIORITY_DEFAULT, 0), ^(void)
@@ -198,6 +199,13 @@ void FFirebaseIOSNotifications::SetFirebaseToken(FString Token)
 {
     @synchronized ([IOSAppDelegate GetDelegate])
     {
+		if (IFirebaseModuleInterface::Get().OnTokenUpdate.IsBound())
+		{
+			FFunctionGraphTask::CreateAndDispatchWhenReady([OldToken = IOSFirebaseToken, NewToken = Token]()
+				{
+					IFirebaseModuleInterface::Get().OnTokenUpdate.Broadcast(OldToken, NewToken);
+				}, TStatId(), NULL, ENamedThreads::GameThread);
+		}
         IOSFirebaseToken = Token;
     }
 }
diff --git a/Engine/Plugins/Runtime/Firebase/Source/Java/notifications/EpicFirebaseMessagingService.java b/Engine/Plugins/Runtime/Firebase/Source/Java/notifications/EpicFirebaseMessagingService.java
--- a/Engine/Plugins/Runtime/Firebase/Source/Java/notifications/EpicFirebaseMessagingService.java
+++ b/Engine/Plugins/Runtime/Firebase/Source/Java/notifications/EpicFirebaseMessagingService.java
@@ -69,6 +69,8 @@ public class EpicFirebaseMessagingService extends FirebaseMessagingService {
 	private static final String KEY_MESSAGE_ID = "ue4_fb.messageId";
 	private static final String KEY_NOTIFICATION_TYPE = "ue4_fb.type";
 
+	public static native void OnFirebaseTokenChange(String PreviousToken, String NewToken);
+
 	@Override
 	public void onMessageReceived(@NonNull RemoteMessage remoteMessage) {
 		Log.verbose("Firebase onMessageReceived");
@@ -283,6 +285,7 @@ public class EpicFirebaseMessagingService extends FirebaseMessagingService {
 	private static void saveFirebaseToken(@NonNull Context context, @NonNull String firebaseToken) {
 		Log.debug("Saving Firebase token");
 		String storedToken = getFirebaseTokenFromCache(context);
+		OnFirebaseTokenChange(storedToken == null ? "" : storedToken, firebaseToken);
 		boolean isUpdatedToken = !TextUtils.isEmpty(storedToken);
 		SharedPreferences sharedPreferences = context.getSharedPreferences(PREFS_FILE_FIREBASE, Context.MODE_PRIVATE);
 		SharedPreferences.Editor editor = sharedPreferences.edit();
@@ -341,6 +344,25 @@ public class EpicFirebaseMessagingService extends FirebaseMessagingService {
 		editor.apply();
 	}
 
+	private static Task<String> queryFirebaseToken(@NonNull Context context, String projectID)
+	{
+		Task<String> tokenTask = FirebaseMessaging.getInstance().getToken();
+		tokenTask.addOnCompleteListener(new OnCompleteListener<String>() {
+			@Override
+			public void onComplete(@NonNull Task<String> task) {
+				if (task.isSuccessful()) {
+					String token = task.getResult();
+					if(!TextUtils.isEmpty(token)) {
+						Log.debug("Firebase token retrieved from Firebase");
+						saveFirebaseToken(context, token);
+						setLastFirebaseProjectID(context, projectID);
+					}
+				}
+			}
+		});
+		return tokenTask;
+	}
+
 	@Nullable
 	public static String getFirebaseToken(@NonNull Context context) {
 		EnsureFirebaseIsInitialized(context);
@@ -360,20 +382,7 @@ public class EpicFirebaseMessagingService extends FirebaseMessagingService {
 		}
 
 		if(TextUtils.isEmpty(token)) {
-			Task<String> tokenTask = FirebaseMessaging.getInstance().getToken();
-			tokenTask.addOnCompleteListener(new OnCompleteListener<String>() {
-				@Override
-				public void onComplete(@NonNull Task<String> task) {
-					if (task.isSuccessful()) {
-						String token = task.getResult();
-						if(!TextUtils.isEmpty(token)) {
-							Log.debug("Firebase token retrieved from Firebase");
-							saveFirebaseToken(context, token);
-							setLastFirebaseProjectID(context, projectID);
-						}
-					}
-				}
-			});
+			Task<String> tokenTask = queryFirebaseToken(context, projectID);
 			// wait on task
 			try {
 				Tasks.await(tokenTask);
@@ -383,7 +392,11 @@ public class EpicFirebaseMessagingService extends FirebaseMessagingService {
 				Log.error("Failed to retrieve Firebase token", e);
 			}
 		} else {
+			OnFirebaseTokenChange("", token);
 			Log.debug("Firebase token retrieved from cache");
+
+			// Schedule token query anyway in case cached token is stale, do it after OnFirebaseTokenChange to make sure we notify with cached token first
+			queryFirebaseToken(context, projectID);
 		}
 
 		if (!Get().nativeIsShippingBuild()) {
diff --git a/Engine/Plugins/Runtime/Firebase/Source/Private/Firebase.cpp b/Engine/Plugins/Runtime/Firebase/Source/Private/Firebase.cpp
--- a/Engine/Plugins/Runtime/Firebase/Source/Private/Firebase.cpp
+++ b/Engine/Plugins/Runtime/Firebase/Source/Private/Firebase.cpp
@@ -17,3 +17,25 @@ class FFirebaseModule : public IFirebaseModuleInterface
 };
 
 IMPLEMENT_MODULE(FFirebaseModule, Firebase);
+
+#if PLATFORM_ANDROID
+#include "Android/AndroidJava.h"
+#include "Android/AndroidJavaEnv.h"
+#include "Android/AndroidJNI.h"
+
+extern "C"
+{
+	JNIEXPORT void Java_com_epicgames_unreal_notifications_EpicFirebaseMessagingService_OnFirebaseTokenChange(JNIEnv* jenv, jobject thiz, jstring jPreviousToken, jstring jNewToken)
+	{
+		if (IFirebaseModuleInterface::Get().OnTokenUpdate.IsBound())
+		{
+			FString PreviousToken = FJavaHelper::FStringFromParam(jenv, jPreviousToken);
+			FString NewToken = FJavaHelper::FStringFromParam(jenv, jNewToken);
+			FFunctionGraphTask::CreateAndDispatchWhenReady([PreviousToken, NewToken]()
+				{
+					IFirebaseModuleInterface::Get().OnTokenUpdate.Broadcast(PreviousToken, NewToken);
+				}, TStatId(), NULL, ENamedThreads::GameThread);
+		}
+	}
+}
+#endif
\ No newline at end of file
diff --git a/Engine/Plugins/Runtime/Firebase/Source/Private/Firebase.h b/Engine/Plugins/Runtime/Firebase/Source/Private/Firebase.h
deleted file mode 100644
diff --git a/Engine/Plugins/Runtime/Firebase/Source/Public/Firebase.h b/Engine/Plugins/Runtime/Firebase/Source/Public/Firebase.h
new file mode 100644
--- /dev/null
+++ b/Engine/Plugins/Runtime/Firebase/Source/Public/Firebase.h
@@ -0,0 +1,43 @@
+// Copyright Epic Games, Inc. All Rights Reserved.
+
+#pragma once
+
+#include "Modules/ModuleInterface.h"
+#include "Modules/ModuleManager.h"
+
+DECLARE_LOG_CATEGORY_EXTERN(LogFirebase, Log, All);
+
+static const char* FIREBASE_MODULE_NAME = "Firebase";
+
+class IFirebaseModuleInterface : public IModuleInterface
+{
+
+public:
+
+	/**
+	 * Singleton-like access to this module's interface.  This is just for convenience!
+	 * Beware of calling this during the shutdown phase - the module might have been unloaded already.
+	 *
+	 * @return Returns singleton instance, loading the module on demand if needed
+	 */
+	static IFirebaseModuleInterface& Get()
+	{
+		return FModuleManager::LoadModuleChecked<IFirebaseModuleInterface>(FIREBASE_MODULE_NAME);
+	}
+
+	/**
+	 * Checks to see if this module is loaded and ready.  It is only valid to call Get() if IsAvailable() returns true.
+	 *
+	 * @return True if the module is loaded and ready to use
+	 */
+	static bool IsAvailable()
+	{
+		return FModuleManager::Get().IsModuleLoaded(FIREBASE_MODULE_NAME);
+	}
+
+	/** IModuleInterface implementation */
+	virtual void StartupModule() override;
+	virtual void ShutdownModule() override;
+	
+	TMulticastDelegate<void(const FString&, const FString&)> OnTokenUpdate;
+};
