Fix for Android non-pch build

diff --git a/Engine/Plugins/Runtime/Firebase/Source/Private/Firebase.cpp b/Engine/Plugins/Runtime/Firebase/Source/Private/Firebase.cpp
--- a/Engine/Plugins/Runtime/Firebase/Source/Private/Firebase.cpp
+++ b/Engine/Plugins/Runtime/Firebase/Source/Private/Firebase.cpp
@@ -22,6 +22,7 @@ IMPLEMENT_MODULE(FFirebaseModule, Firebase);
 #include "Android/AndroidJava.h"
 #include "Android/AndroidJavaEnv.h"
 #include "Android/AndroidJNI.h"
+#include "Async/TaskGraphInterfaces.h"
 
 extern "C"
 {
