Fixes missing include

#rnx

diff --git a/Engine/Plugins/Runtime/Firebase/Source/IOS/notifications/EpicFirebaseIOSNotifications.cpp b/Engine/Plugins/Runtime/Firebase/Source/IOS/notifications/EpicFirebaseIOSNotifications.cpp
--- a/Engine/Plugins/Runtime/Firebase/Source/IOS/notifications/EpicFirebaseIOSNotifications.cpp
+++ b/Engine/Plugins/Runtime/Firebase/Source/IOS/notifications/EpicFirebaseIOSNotifications.cpp
@@ -14,6 +14,9 @@ THIRD_PARTY_INCLUDES_END
 #include "Firebase.h"
 #include "IOS/IOSAppDelegate.h"
 
+#include "Misc/Paths.h"
+#include "Async/TaskGraphInterfaces.h"
+
 bool FFirebaseIOSNotifications::bIsInitialized = false;
 FString FFirebaseIOSNotifications::IOSFirebaseToken;
 NSString* KEY_FIREBASE_TOKEN = @"firebasetoken";
