Ensure java sources are wiped with FirebaseEnabled is toggled off. Otherwise, source files are left behind that don't compile as the Java dependencies don't get populated by the plugin when FirebaseEnabled = false.
Additionally, registerBuildSettiings for the Firebase settings referenced from the plugin's UPL.

diff --git a/Engine/Plugins/Runtime/Firebase/Source/Firebase.upl.xml b/Engine/Plugins/Runtime/Firebase/Source/Firebase.upl.xml
--- a/Engine/Plugins/Runtime/Firebase/Source/Firebase.upl.xml
+++ b/Engine/Plugins/Runtime/Firebase/Source/Firebase.upl.xml
@@ -18,6 +18,11 @@
 		</if>
 	</init>
 
+	<registerBuildSettings>
+		<insertValue value="Firebase[FirebaseEnabled,FirebaseDisableAutoInit,FirebaseService]"/>
+		<insertNewline/>
+	</registerBuildSettings>
+	
 	<androidManifestUpdates>
 		<if condition="bEnabled">
 			<true>
@@ -80,6 +85,9 @@ apply plugin: 'com.google.gms.google-services'
 			<true>
 				<copyDir src="$S(PluginDir)/Java" dst="$S(BuildDir)/src/com/epicgames/unreal"/>
 			</true>
+			<false>
+				<deleteFiles filespec="$S(BuildDir)/src/com/epicgames/unreal/notifications/*.java"/>
+			</false>
 		</if>
 	</prebuildCopies>
 	
