Firebase will log project ID in non-shipping builds

diff --git a/Engine/Plugins/Runtime/Firebase/Source/IOS/notifications/EpicFirebaseIOSNotifications.cpp b/Engine/Plugins/Runtime/Firebase/Source/IOS/notifications/EpicFirebaseIOSNotifications.cpp
--- a/Engine/Plugins/Runtime/Firebase/Source/IOS/notifications/EpicFirebaseIOSNotifications.cpp
+++ b/Engine/Plugins/Runtime/Firebase/Source/IOS/notifications/EpicFirebaseIOSNotifications.cpp
@@ -60,6 +60,11 @@ NSString* KEY_FIREBASE_PROJECT_ID = @"firebaseprojectid";
 -(void)ConfigureFirebase
 {
     [FIRApp configure];
+	
+#if !UE_BUILD_SHIPPING
+	NSString* currentProjectID = [[[FIRApp defaultApp] options] GCMSenderID];
+	NSLog(@"Firebase configured for project %@", currentProjectID);
+#endif
 }
 
 -(void)ConfigureFirebaseWithCustomFile : (NSString*)fileName fileExtension:(NSString*)ext
@@ -81,6 +86,11 @@ NSString* KEY_FIREBASE_PROJECT_ID = @"firebaseprojectid";
 		else
 		{
 			[FIRApp configureWithOptions:options];
+			
+#if !UE_BUILD_SHIPPING
+	NSString* currentProjectID = [[[FIRApp defaultApp] options] GCMSenderID];
+	NSLog(@"Firebase configured using custom config for project %@", currentProjectID);
+#endif
 		}
 	}
 }
diff --git a/Engine/Plugins/Runtime/Firebase/Source/Java/notifications/EpicFirebaseMessagingService.java b/Engine/Plugins/Runtime/Firebase/Source/Java/notifications/EpicFirebaseMessagingService.java
--- a/Engine/Plugins/Runtime/Firebase/Source/Java/notifications/EpicFirebaseMessagingService.java
+++ b/Engine/Plugins/Runtime/Firebase/Source/Java/notifications/EpicFirebaseMessagingService.java
@@ -418,6 +418,10 @@ public class EpicFirebaseMessagingService extends FirebaseMessagingService {
 	public static void EnsureFirebaseIsInitialized(@NonNull Context context) {
 		if (FirebaseApp.getApps(context).isEmpty()) {
 			FirebaseApp.initializeApp(context);
+			if (!GameActivity.IS_SHIPPING_CONFIG)
+			{
+				Log.debug("Initialized Firebase for project : " + FirebaseApp.getInstance().getOptions().getGcmSenderId());
+			}
 		}
 	}
 
@@ -483,6 +487,11 @@ public class EpicFirebaseMessagingService extends FirebaseMessagingService {
 			FirebaseOptions options = optionsBuilder.build();
 
 			FirebaseApp.initializeApp(context, options);
+
+			if (!GameActivity.IS_SHIPPING_CONFIG)
+			{
+				Log.debug("Initialized Firebase with custom file for project : " + FirebaseApp.getInstance().getOptions().getGcmSenderId());
+			}
 		}
 		catch(Exception e)
 		{
