iOS Firebase set APNS token manually

diff --git a/Engine/Plugins/Runtime/Firebase/Source/IOS/notifications/EpicFirebaseIOSNotifications.cpp b/Engine/Plugins/Runtime/Firebase/Source/IOS/notifications/EpicFirebaseIOSNotifications.cpp
--- a/Engine/Plugins/Runtime/Firebase/Source/IOS/notifications/EpicFirebaseIOSNotifications.cpp
+++ b/Engine/Plugins/Runtime/Firebase/Source/IOS/notifications/EpicFirebaseIOSNotifications.cpp
@@ -25,7 +25,7 @@ NSString* KEY_FIREBASE_PROJECT_ID = @"firebaseprojectid";
 @interface IOSAppDelegate (FirebaseHandling) <FIRMessagingDelegate>
 
 -(void)SetupFirebase : (Boolean) enableAnalytics;
--(void)UpdateFirebaseToken : (UInt64) Timeout;
+-(void)ScheduleFirebaseTokenUpdate;
 
 @end
 
@@ -40,21 +40,22 @@ NSString* KEY_FIREBASE_PROJECT_ID = @"firebaseprojectid";
     [FIRMessaging messaging].delegate = self;
     [UNUserNotificationCenter currentNotificationCenter].delegate = self;
     
-    UNAuthorizationOptions authOptions = UNAuthorizationOptionAlert | UNAuthorizationOptionSound | UNAuthorizationOptionBadge;
-    [[UNUserNotificationCenter currentNotificationCenter] requestAuthorizationWithOptions:authOptions
-        completionHandler:^(BOOL granted, NSError * _Nullable error) 
-     {
-        if (granted)
-        {
-            UE_LOG(LogFirebase, Log, TEXT("Firebase authorization granted"));
-        }
-        else
-        {
-            UE_LOG(LogFirebase, Log, TEXT("Firebase authorization denied"));
-        }
-    }];
-
-   [[UIApplication sharedApplication] registerForRemoteNotifications];
+	
+	UNAuthorizationOptions authOptions = UNAuthorizationOptionAlert | UNAuthorizationOptionSound | UNAuthorizationOptionBadge;
+	[[UNUserNotificationCenter currentNotificationCenter] requestAuthorizationWithOptions:authOptions
+		completionHandler:^(BOOL granted, NSError * _Nullable error) 
+	 {
+		if (granted)
+		{
+			UE_LOG(LogFirebase, Log, TEXT("Firebase authorization granted"));
+			
+			[[UIApplication sharedApplication] registerForRemoteNotifications];
+		}
+		else
+		{
+			UE_LOG(LogFirebase, Log, TEXT("Firebase authorization denied"));
+		}
+	}];
 }
 
 -(void)ConfigureFirebase
@@ -118,7 +119,7 @@ NSString* KEY_FIREBASE_PROJECT_ID = @"firebaseprojectid";
     [[NSUserDefaults standardUserDefaults] synchronize];
 }
 
-- (void)UpdateFirebaseToken : (UInt64) Timeout
+- (void)ScheduleFirebaseTokenUpdate
 {
     NSUserDefaults* UserDefaults = [NSUserDefaults standardUserDefaults];
 	NSString* currentProjectID = [[[FIRApp defaultApp] options] GCMSenderID];
@@ -143,30 +144,22 @@ NSString* KEY_FIREBASE_PROJECT_ID = @"firebaseprojectid";
     
 	// Query token from Firebase even if there is one in cache already
 	
-    dispatch_semaphore_t updateTokenSemaphore = dispatch_semaphore_create(0);
-    // wrapped in dispatch_async to avoid locking up if we're on the main thread
-    dispatch_async(dispatch_get_global_queue(DISPATCH_QUEUE_PRIORITY_DEFAULT, 0), ^(void)
-    {
-        [[FIRMessaging messaging] tokenWithCompletion:^(NSString *firebaseToken, NSError *error)
-         {
-            if (error == nil && firebaseToken != nil)
-            {
-                FString Token = FString(firebaseToken);
-                FFirebaseIOSNotifications::SetFirebaseToken(Token);
+	
+	[[FIRMessaging messaging] tokenWithCompletion:^(NSString *firebaseToken, NSError *error)
+	 {
+		if (error == nil && firebaseToken != nil)
+		{
+			FString Token = FString(firebaseToken);
+			FFirebaseIOSNotifications::SetFirebaseToken(Token);
 #if !UE_BUILD_SHIPPING
-                UE_LOG(LogFirebase, Log, TEXT("Firebase Token Queried : %s"), *Token);
+			UE_LOG(LogFirebase, Log, TEXT("Firebase Token Queried : %s"), *Token);
 #endif
-				NSString* projectID = [[[FIRApp defaultApp] options] GCMSenderID];
-                [[NSUserDefaults standardUserDefaults] setObject:firebaseToken forKey:KEY_FIREBASE_TOKEN];
-                [[NSUserDefaults standardUserDefaults] setObject:projectID forKey:KEY_FIREBASE_PROJECT_ID];
-                [[NSUserDefaults standardUserDefaults] synchronize];
-                
-                dispatch_semaphore_signal(updateTokenSemaphore);
-            }
-        }];
-    });
-    dispatch_semaphore_wait(updateTokenSemaphore, dispatch_time(DISPATCH_TIME_NOW, Timeout));
-    dispatch_release(updateTokenSemaphore);
+			NSString* projectID = [[[FIRApp defaultApp] options] GCMSenderID];
+			[[NSUserDefaults standardUserDefaults] setObject:firebaseToken forKey:KEY_FIREBASE_TOKEN];
+			[[NSUserDefaults standardUserDefaults] setObject:projectID forKey:KEY_FIREBASE_PROJECT_ID];
+			[[NSUserDefaults standardUserDefaults] synchronize];
+		}
+	}];
 }
 @end
 
@@ -203,9 +196,13 @@ void FFirebaseIOSNotifications::Initialize(uint64 TokenQueryTimeoutNanoseconds,
     if (!bIsInitialized)
     {
         [[IOSAppDelegate GetDelegate] SetupFirebase:bEnableAnalytics];
-        [[IOSAppDelegate GetDelegate] UpdateFirebaseToken:TokenQueryTimeoutNanoseconds];
         bIsInitialized = true;
     }
+	else 
+	{
+		// Ensure we still have a valid token if already initialized
+		[[IOSAppDelegate GetDelegate] ScheduleFirebaseTokenUpdate];
+	}
 }
 
 void FFirebaseIOSNotifications::SetFirebaseToken(FString Token)
@@ -231,11 +228,6 @@ FString FFirebaseIOSNotifications::GetFirebaseToken()
         Token = IOSFirebaseToken;
     }
     
-    if (Token.IsEmpty())
-    {
-        UE_LOG(LogFirebase, Log, TEXT("Firebase Token is empty"));
-    }
-    
     return Token;
 }
 
@@ -268,4 +260,16 @@ void FFirebaseIOSNotifications::DeleteFirebaseToken()
 	}];
 }
 
+void FFirebaseIOSNotifications::OnAPNSTokenReceived(TArrayView<uint8> Token)
+{
+	NSData* apnsToken = Token.IsEmpty() ? nil : [NSData dataWithBytes:Token.GetData() length:Token.Num()];
+	
+	if (apnsToken != nil)
+	{
+		[[FIRMessaging messaging] setAPNSToken:apnsToken];
+	}
+	
+	[[IOSAppDelegate GetDelegate] ScheduleFirebaseTokenUpdate];
+}
+
 #endif // PLATFORM_IOS
diff --git a/Engine/Plugins/Runtime/Firebase/Source/IOS/notifications/EpicFirebaseIOSNotifications.h b/Engine/Plugins/Runtime/Firebase/Source/IOS/notifications/EpicFirebaseIOSNotifications.h
--- a/Engine/Plugins/Runtime/Firebase/Source/IOS/notifications/EpicFirebaseIOSNotifications.h
+++ b/Engine/Plugins/Runtime/Firebase/Source/IOS/notifications/EpicFirebaseIOSNotifications.h
@@ -27,6 +27,7 @@ public:
     static void SetFirebaseToken(FString Token);
     static FString GetFirebaseToken();
 	static void DeleteFirebaseToken();
+	static void OnAPNSTokenReceived(TArrayView<uint8> Token);
 };
 
 #endif // PLATFORM_IOS
diff --git a/Engine/Plugins/Runtime/Firebase/Source/Private/Firebase.cpp b/Engine/Plugins/Runtime/Firebase/Source/Private/Firebase.cpp
--- a/Engine/Plugins/Runtime/Firebase/Source/Private/Firebase.cpp
+++ b/Engine/Plugins/Runtime/Firebase/Source/Private/Firebase.cpp
@@ -4,12 +4,25 @@
 
 DEFINE_LOG_CATEGORY(LogFirebase);
 
+#if PLATFORM_IOS
+#include "IOS/notifications/EpicFirebaseIOSNotifications.h"
+#endif
+
 void IFirebaseModuleInterface::StartupModule()
 {
+#if PLATFORM_IOS
+	ApplicationRegisteredForRemoteNotificationsHandle = FCoreDelegates::ApplicationRegisteredForRemoteNotificationsDelegate.AddLambda([](TArray<uint8> Token){
+		FFirebaseIOSNotifications::OnAPNSTokenReceived(Token);
+	});
+#endif
 }
 
 void IFirebaseModuleInterface::ShutdownModule()
 {
+#if PLATFORM_IOS
+	FCoreDelegates::ApplicationRegisteredForRemoteNotificationsDelegate.Remove(ApplicationRegisteredForRemoteNotificationsHandle);
+	ApplicationRegisteredForRemoteNotificationsHandle = {};
+#endif
 }
 
 class FFirebaseModule : public IFirebaseModuleInterface
@@ -39,4 +52,4 @@ extern "C"
 		}
 	}
 }
-#endif
\ No newline at end of file
+#endif
diff --git a/Engine/Plugins/Runtime/Firebase/Source/Public/Firebase.h b/Engine/Plugins/Runtime/Firebase/Source/Public/Firebase.h
--- a/Engine/Plugins/Runtime/Firebase/Source/Public/Firebase.h
+++ b/Engine/Plugins/Runtime/Firebase/Source/Public/Firebase.h
@@ -11,6 +11,10 @@ static const char* FIREBASE_MODULE_NAME = "Firebase";
 
 class IFirebaseModuleInterface : public IModuleInterface
 {
+private:
+#if PLATFORM_IOS
+	FDelegateHandle ApplicationRegisteredForRemoteNotificationsHandle;
+#endif
 
 public:
 
