Added missing Firebase iOS integration check

diff --git a/Engine/Plugins/Runtime/Firebase/Source/Private/Firebase.cpp b/Engine/Plugins/Runtime/Firebase/Source/Private/Firebase.cpp
--- a/Engine/Plugins/Runtime/Firebase/Source/Private/Firebase.cpp
+++ b/Engine/Plugins/Runtime/Firebase/Source/Private/Firebase.cpp
@@ -10,7 +10,7 @@ DEFINE_LOG_CATEGORY(LogFirebase);
 
 void IFirebaseModuleInterface::StartupModule()
 {
-#if PLATFORM_IOS
+#if PLATFORM_IOS && WITH_IOS_FIREBASE_INTEGRATION
 	ApplicationRegisteredForRemoteNotificationsHandle = FCoreDelegates::ApplicationRegisteredForRemoteNotificationsDelegate.AddLambda([](TArray<uint8> Token){
 		FFirebaseIOSNotifications::OnAPNSTokenReceived(Token);
 	});
@@ -19,7 +19,7 @@ void IFirebaseModuleInterface::StartupModule()
 
 void IFirebaseModuleInterface::ShutdownModule()
 {
-#if PLATFORM_IOS
+#if PLATFORM_IOS && WITH_IOS_FIREBASE_INTEGRATION
 	FCoreDelegates::ApplicationRegisteredForRemoteNotificationsDelegate.Remove(ApplicationRegisteredForRemoteNotificationsHandle);
 	ApplicationRegisteredForRemoteNotificationsHandle = {};
 #endif
diff --git a/Engine/Plugins/Runtime/Firebase/Source/Public/Firebase.h b/Engine/Plugins/Runtime/Firebase/Source/Public/Firebase.h
--- a/Engine/Plugins/Runtime/Firebase/Source/Public/Firebase.h
+++ b/Engine/Plugins/Runtime/Firebase/Source/Public/Firebase.h
@@ -12,7 +12,7 @@ static const char* FIREBASE_MODULE_NAME = "Firebase";
 class IFirebaseModuleInterface : public IModuleInterface
 {
 private:
-#if PLATFORM_IOS
+#if PLATFORM_IOS && WITH_IOS_FIREBASE_INTEGRATION
 	FDelegateHandle ApplicationRegisteredForRemoteNotificationsHandle;
 #endif
 
