iOS non unity build fix

diff --git a/Engine/Plugins/Runtime/Firebase/Source/Public/Firebase.h b/Engine/Plugins/Runtime/Firebase/Source/Public/Firebase.h
--- a/Engine/Plugins/Runtime/Firebase/Source/Public/Firebase.h
+++ b/Engine/Plugins/Runtime/Firebase/Source/Public/Firebase.h
@@ -4,6 +4,7 @@
 
 #include "Modules/ModuleInterface.h"
 #include "Modules/ModuleManager.h"
+#include "Delegates/IDelegateInstance.h"
 
 DECLARE_LOG_CATEGORY_EXTERN(LogFirebase, Log, All);
 
@@ -12,9 +13,7 @@ static const char* FIREBASE_MODULE_NAME = "Firebase";
 class IFirebaseModuleInterface : public IModuleInterface
 {
 private:
-#if PLATFORM_IOS && WITH_IOS_FIREBASE_INTEGRATION
 	FDelegateHandle ApplicationRegisteredForRemoteNotificationsHandle;
-#endif
 
 public:
 
