IWYU unity compilation

diff --git a/Engine/Plugins/Runtime/Firebase/Source/Private/Firebase.cpp b/Engine/Plugins/Runtime/Firebase/Source/Private/Firebase.cpp
--- a/Engine/Plugins/Runtime/Firebase/Source/Private/Firebase.cpp
+++ b/Engine/Plugins/Runtime/Firebase/Source/Private/Firebase.cpp
@@ -1,6 +1,7 @@
 // Copyright Epic Games, Inc. All Rights Reserved.
 
 #include "Firebase.h"
+#include "Misc/CoreDelegates.h"
 
 DEFINE_LOG_CATEGORY(LogFirebase);
 
